Added a DerivedDataBuildController to support compiling shaders through the Derived Data Build API

Add -DDBShaderCompile to the editor/cooker command line to enable. Include -NoUbaShaderCompile if needed to disable UBA and -NoXGEShaderCompile to disable XGE.

Integrates at the level of IDistributedBuildController. The Derived Data Build API would ideally integrate at the level of compiling a single shader and will be evaluated later.

Versioning is incomplete. This checks shader format versions like SCW but does not provide a distinct function version that tracks changes to shader compilation.

#Synced-CL 50253651

diff --git a/Engine/Plugins/DerivedDataBuildController/DerivedDataBuildController.uplugin b/Engine/Plugins/DerivedDataBuildController/DerivedDataBuildController.uplugin
new file mode 100644
--- /dev/null
+++ b/Engine/Plugins/DerivedDataBuildController/DerivedDataBuildController.uplugin
@@ -0,0 +1,30 @@
+{
+	"FileVersion" : 1,
+	"Version" : 1,
+	"VersionName" : "0.1",
+	"FriendlyName" : "Derived Data Build Controller",
+	"Description" : "Adds support for shader compiling distribution using the Derived Data Build API",
+	"Category" : "Build Distribution",
+	"CreatedBy" : "Epic Games, Inc.",
+	"CreatedByURL" : "https://epicgames.com",
+	"DocsURL" : "",
+	"MarketplaceURL" : "",
+	"SupportURL" : "",
+	"IsExperimentalVersion": true,
+	"EnabledByDefault" : true,
+	"CanContainContent" : false,
+	"IsBetaVersion" : false,
+	"Installed" : false,
+
+	"Modules" :
+	[
+		{
+			"Name" : "DerivedDataBuildController",
+			"Type" : "Editor",
+			"LoadingPhase" : "EarliestPossible",
+			"PlatformAllowList": ["Win64", "Mac", "Linux"]
+		}
+	],
+
+	"SupportedTargetPlatforms": ["Win64", "Mac", "Linux"]
+}
diff --git a/Engine/Plugins/DerivedDataBuildController/Source/DerivedDataBuildController.Build.cs b/Engine/Plugins/DerivedDataBuildController/Source/DerivedDataBuildController.Build.cs
new file mode 100644
--- /dev/null
+++ b/Engine/Plugins/DerivedDataBuildController/Source/DerivedDataBuildController.Build.cs
@@ -0,0 +1,17 @@
+// Copyright Epic Games, Inc. All Rights Reserved.
+
+using System.IO;
+using UnrealBuildTool;
+
+public class DerivedDataBuildController : ModuleRules
+{
+	public DerivedDataBuildController(ReadOnlyTargetRules TargetRules) : base(TargetRules)
+	{
+		PrivateDependencyModuleNames.AddRange(new string[]
+		{
+			"Core",
+			"DerivedDataCache",
+			"DistributedBuildInterface",
+		});
+	}
+}
diff --git a/Engine/Plugins/DerivedDataBuildController/Source/Private/DerivedDataBuildControllerModule.cpp b/Engine/Plugins/DerivedDataBuildController/Source/Private/DerivedDataBuildControllerModule.cpp
new file mode 100644
--- /dev/null
+++ b/Engine/Plugins/DerivedDataBuildController/Source/Private/DerivedDataBuildControllerModule.cpp
@@ -0,0 +1,243 @@
+// Copyright Epic Games, Inc. All Rights Reserved.
+
+#include "Compression/CompressedBuffer.h"
+#include "Compression/OodleDataCompression.h"
+#include "DerivedDataBuild.h"
+#include "DerivedDataBuildAction.h"
+#include "DerivedDataBuildInputs.h"
+#include "DerivedDataBuildOutput.h"
+#include "DerivedDataBuildSession.h"
+#include "DerivedDataRequestOwner.h"
+#include "DerivedDataValue.h"
+#include "DistributedBuildControllerInterface.h"
+#include "Features/IModularFeatures.h"
+#include "IO/IoHash.h"
+#include "Logging/LogMacros.h"
+#include "Misc/CommandLine.h"
+#include "Misc/CoreDelegates.h"
+#include "Misc/Paths.h"
+#include "Modules/ModuleManager.h"
+
+// Used to register the remote function.
+#include "DerivedDataBuildFunction.h"
+#include "DerivedDataBuildFunctionFactory.h"
+
+DEFINE_LOG_CATEGORY_STATIC(LogDerivedDataBuildController, Log, All);
+
+class FDerivedDataBuildControllerModule final : public IDistributedBuildController , public TSharedFromThis<FDerivedDataBuildControllerModule>
+{
+public:
+	FDerivedDataBuildControllerModule() = default;
+	~FDerivedDataBuildControllerModule() final = default;
+
+	// IModuleInterface
+
+	void StartupModule() final;
+	void ShutdownModule() final;
+
+	// IDistributedBuildController
+
+	bool SupportsLocalWorkers() final { return true; }
+	bool SupportsVirtualFiles() const final { return true; }
+
+	inline int32 GetMaxNumLocalWorkers() const { return MaxNumLocalWorkers; }
+	void SetMaxLocalWorkers(int32 InMaxNumLocalWorkers) final { MaxNumLocalWorkers = InMaxNumLocalWorkers; }
+
+	void InitializeController() final;
+	bool IsSupported() final;
+
+	const FString GetName() final { return "Derived Data Build Controller"; }
+
+	FString CreateUniqueFilePath() final;
+	bool PollStats(FDistributedBuildStats& OutStats) final;
+	TFuture<FDistributedBuildTaskResult> EnqueueTask(const FTaskCommandData& CommandData) final;
+
+	static const UE::FUtf8SharedString& GetFunctionName()
+	{
+		static const UE::FUtf8SharedString Name(ANSITEXTVIEW("CompileShaderJobs"));
+		return Name;
+	}
+
+private:
+	bool bSupported = false;
+	bool bControllerInitialized = false;
+	
+	TAtomic<int32> MaxNumLocalWorkers = -1;
+
+	UE::DerivedData::FOptionalBuildSession Session;
+};
+
+IMPLEMENT_MODULE(FDerivedDataBuildControllerModule, DerivedDataBuildController);
+
+namespace DerivedDataBuildControllerModule
+{
+	static constexpr int32 SubFolderCount = 32;
+}
+
+static bool IsDerivedDataBuildControllerEnabled()
+{
+	if (FParse::Param(FCommandLine::Get(), TEXT("NoDDBShaderCompile")) ||
+		FParse::Param(FCommandLine::Get(), TEXT("NoShaderWorker")))
+	{
+		return false;
+	}
+
+	if (FParse::Param(FCommandLine::Get(), TEXT("DDBShaderCompile")))
+	{
+		return true;
+	}
+
+	return false;
+}
+
+bool FDerivedDataBuildControllerModule::IsSupported()
+{
+	if (!bControllerInitialized)
+	{
+		bSupported = IsDerivedDataBuildControllerEnabled() && FPlatformProcess::SupportsMultithreading();
+	}
+	return bSupported;
+}
+
+void FDerivedDataBuildControllerModule::StartupModule()
+{
+	IModularFeatures::Get().RegisterModularFeature(GetModularFeatureType(), this);
+}
+
+void FDerivedDataBuildControllerModule::ShutdownModule()
+{
+	IModularFeatures::Get().UnregisterModularFeature(GetModularFeatureType(), this);
+	
+	if (bControllerInitialized)
+	{
+		Session.Reset();
+		bControllerInitialized = false;
+	}
+}
+
+void FDerivedDataBuildControllerModule::InitializeController()
+{
+	using namespace UE;
+	using namespace UE::DerivedData;
+
+	if (ensureAlwaysMsgf(!bControllerInitialized, TEXT("Multiple initialization of Derived Data Build controller!")))
+	{
+		if (IsSupported())
+		{
+			FModuleManager::Get().LoadModuleChecked(TEXT("DerivedDataCache"));
+			Session = UE::DerivedData::GetBuild().CreateSession(TEXTVIEW("DerivedDataBuildControllerModule"));
+		}
+
+		bControllerInitialized = true;
+	}
+}
+
+FString FDerivedDataBuildControllerModule::CreateUniqueFilePath()
+{
+	// These paths are never used because of virtual files, but paths are created anyway.
+	static const FString WorkingDirectory = FPaths::Combine(FPlatformProcess::UserTempDir(), TEXT("DerivedDataBuild"), FGuid::NewGuid().ToString(EGuidFormats::DigitsWithHyphensLower));
+	static int32 NextFileID = 0;
+	return FPaths::Combine(WorkingDirectory, FString::Printf(TEXT("%d.tmp"), NextFileID++));
+}
+
+TFuture<FDistributedBuildTaskResult> FDerivedDataBuildControllerModule::EnqueueTask(const FTaskCommandData& CommandData)
+{
+	using namespace UE;
+	using namespace UE::DerivedData;
+
+	check(bSupported);
+
+	const FSharedString Name(CommandData.Description);
+
+	IBuild& Build = GetBuild();
+	FBuildActionBuilder ActionBuilder = Build.CreateAction(Name, GetFunctionName());
+	FBuildInputsBuilder InputsBuilder = Build.CreateInputs(Name);
+
+	{
+		const FCompressedBuffer Input = FCompressedBuffer::Compress(FSharedBuffer::MakeView(CommandData.InputFileData));
+		ActionBuilder.AddInput(ANSITEXTVIEW("Input"), Input.GetRawHash(), Input.GetRawSize());
+		InputsBuilder.AddInput(ANSITEXTVIEW("Input"), Input);
+	}
+
+	for (int32 ViewIndex = 0, ViewCount = CommandData.SourceData.Num(); ViewIndex < ViewCount; ++ViewIndex)
+	{
+		const FCompressedBuffer Input = FCompressedBuffer::Compress(FSharedBuffer::MakeView(CommandData.SourceData[ViewIndex]),
+			ECompressedBufferCompressor::NotSet, ECompressedBufferCompressionLevel::None);
+		TAnsiStringBuilder<32> InputName(InPlace, ANSITEXTVIEW("Virtual"), ViewIndex);
+		ActionBuilder.AddInput(InputName, Input.GetRawHash(), Input.GetRawSize());
+		InputsBuilder.AddInput(InputName, Input);
+	}
+
+	TPromise<FDistributedBuildTaskResult> Promise;
+	TFuture<FDistributedBuildTaskResult> Future = Promise.GetFuture();
+
+	FRequestOwner Owner(EPriority::Normal);
+	Owner.KeepAlive();
+	Session.Get().Build(ActionBuilder.Build(), InputsBuilder.Build(), EBuildPolicy::Build, Owner,
+		[Promise = MoveTemp(Promise)](FBuildCompleteParams&& Params) mutable
+		{
+			FDistributedBuildTaskResult Result;
+			Result.bCompleted = Params.Status != EStatus::Canceled;
+
+			static const FValueId OutputId = FValueId::FromName(ANSITEXTVIEW("Output"));
+			if (const FValueWithId& Value = Params.Output.GetValue(OutputId))
+			{
+				TArray<uint8> OutputData;
+				if (const FCompressedBuffer& CompressedData = Value.GetData())
+				{
+					const uint64 RawSize = CompressedData.GetRawSize();
+					if (RawSize <= MAX_int32)
+					{
+						OutputData.SetNumUninitialized((int32)RawSize);
+						if (CompressedData.TryDecompressTo(OutputData))
+						{
+							Result.OutputData = MoveTemp(OutputData);
+						}
+					}
+				}
+			}
+
+			Result.ReturnCode = Result.OutputData.IsEmpty() ? 1 : 0;
+			Promise.SetValue(MoveTemp(Result));
+		});
+
+	return MoveTemp(Future);
+}
+
+bool FDerivedDataBuildControllerModule::PollStats(FDistributedBuildStats& OutStats)
+{
+	return false;
+}
+
+class FCompileShaderJobsBuildFunction final : public UE::DerivedData::IBuildFunction
+{
+public:
+	const UE::FUtf8SharedString& GetName() const final
+	{
+		return FDerivedDataBuildControllerModule::GetFunctionName();
+	}
+
+	FGuid GetVersion() const final
+	{
+		// This version must match ShaderBuildWorker.cpp.
+		// Version is not yet meaningfully derived from the shader build code,
+		// which means that it cannot yet reliably distinguish two build workers.
+		static const FGuid Version(TEXT("83027356-2cf7-41ca-aba5-c81ab0ff2129"));
+		return Version;
+	}
+
+	void Configure(UE::DerivedData::FBuildConfigContext& Context) const final
+	{
+		// Register the function for remote builds only.
+		// Building in-process will eventually be supported by a build function
+		// that builds one shader at a time instead of batches.
+		Context.SetBuildPolicyMask(~UE::DerivedData::EBuildPolicy::BuildLocal);
+	}
+
+	void Build(UE::DerivedData::FBuildContext& Context) const final
+	{
+		checkNoEntry();
+	}
+};
+
+static UE::DerivedData::TBuildFunctionFactory<FCompileShaderJobsBuildFunction> CompileShaderJobsFunction;
diff --git a/Engine/Source/Editor/UnrealEd/Private/DerivedDataBuildWorkers.cpp b/Engine/Source/Editor/UnrealEd/Private/DerivedDataBuildWorkers.cpp
--- a/Engine/Source/Editor/UnrealEd/Private/DerivedDataBuildWorkers.cpp
+++ b/Engine/Source/Editor/UnrealEd/Private/DerivedDataBuildWorkers.cpp
@@ -9,7 +9,8 @@
 
 struct FDerivedDataBuildWorkers
 {
-	FMultiPlatformTargetReceiptBuildWorkers BaseTextureBuildWorkerFactory {TEXT("$(EngineDir)/Binaries/$(Platform)/BaseTextureBuildWorker.Target")};
+	FMultiPlatformTargetReceiptBuildWorkers BaseTextureBuildWorkerFactory{TEXT("$(EngineDir)/Binaries/$(Platform)/BaseTextureBuildWorker.target")};
+	FMultiPlatformTargetReceiptBuildWorkers ShaderBuildWorkerFactory{TEXT("$(EngineDir)/Binaries/$(Platform)/ShaderBuildWorker.target")};
 };
 
 TOptional<FDerivedDataBuildWorkers> GBuildWorkers;
diff --git a/Engine/Source/Programs/ShaderBuildWorker/Private/ShaderBuildWorker.cpp b/Engine/Source/Programs/ShaderBuildWorker/Private/ShaderBuildWorker.cpp
new file mode 100644
--- /dev/null
+++ b/Engine/Source/Programs/ShaderBuildWorker/Private/ShaderBuildWorker.cpp
@@ -0,0 +1,116 @@
+// Copyright Epic Games, Inc. All Rights Reserved.
+
+#include "Containers/Array.h"
+#include "Containers/UnrealString.h"
+#include "DerivedDataBuildFunction.h"
+#include "DerivedDataBuildFunctionFactory.h"
+#include "DerivedDataValue.h"
+#include "Logging/StructuredLog.h"
+#include "Serialization/MemoryReader.h"
+#include "Serialization/MemoryWriter.h"
+#include "ShaderCompileWorkerUtil.h"
+#include "SocketSubsystem.h"
+
+class FCompileShaderJobsBuildFunction final : public UE::DerivedData::IBuildFunction
+{
+public:
+	const UE::FUtf8SharedString& GetName() const final
+	{
+		static const UE::FUtf8SharedString Name(ANSITEXTVIEW("CompileShaderJobs"));
+		return Name;
+	}
+
+	FGuid GetVersion() const final
+	{
+		// This version must match DerivedDataBuildControllerModule.cpp.
+		// Version is not yet meaningfully derived from the shader build code,
+		// which means that it cannot yet reliably distinguish two build workers.
+		static const FGuid Version(TEXT("83027356-2cf7-41ca-aba5-c81ab0ff2129"));
+		return Version;
+	}
+
+	void Configure(UE::DerivedData::FBuildConfigContext& Context) const final
+	{
+	}
+
+	void Build(UE::DerivedData::FBuildContext& Context) const final;
+};
+
+void FCompileShaderJobsBuildFunction::Build(UE::DerivedData::FBuildContext& Context) const
+{
+	using namespace UE::DerivedData;
+
+	FShaderCompileWorkerDiagnostics WorkerDiagnostics;
+
+	int32 NumProcessedJobs = 0;
+	TArray<FShaderCompileJob> SingleJobs;
+	TArray<FShaderPipelineCompileJob> PipelineJobs;
+	TArray<FString> PipelineJobNames;
+
+	if (FShaderCompileWorkerUtil::GetShaderFormats().IsEmpty())
+	{
+		FSCWErrorCode::Report(FSCWErrorCode::NoTargetShaderFormatsFound, TEXT("No target shader formats found!"));
+	}
+	else
+	{
+		WorkerDiagnostics.EntryPointTimestamp = FPlatformTime::Seconds();
+
+		if (FSharedBuffer Input = Context.FindInput(ANSITEXTVIEW("Input")))
+		{
+			const double BatchProcessStartTime = FPlatformTime::Seconds();
+			WorkerDiagnostics.BatchPreparationTime = BatchProcessStartTime - WorkerDiagnostics.EntryPointTimestamp;
+
+			FMemoryReaderView InputAr(Input, /*bIsPersistent*/ true);
+
+			// TODO: Passed to SCW as -DebugSourceFiles=<Path1>[,<Path2>...] relative to WorkingDirectory
+			TArray<FString> DebugSourceFiles;
+
+			auto LoadVirtualSourceFile = [&Context](int32 VirtualSourceIndex) -> TArray<uint8>
+			{
+				TArray<uint8> Output;
+				if (FSharedBuffer Data = Context.FindInput(WriteToAnsiString<16>(ANSITEXTVIEW("Virtual"), VirtualSourceIndex)))
+				{
+					Output.Append((const uint8*)Data.GetData(), (int32)Data.GetSize());
+				}
+				return Output;
+			};
+
+			FShaderCompileWorkerUtil::ProcessInputFromArchive(&InputAr, {}, DebugSourceFiles, LoadVirtualSourceFile,
+				NumProcessedJobs, SingleJobs, PipelineJobs, PipelineJobNames);
+
+			WorkerDiagnostics.BatchProcessTime = FPlatformTime::Seconds() - BatchProcessStartTime;
+		}
+		else
+		{
+			FSCWErrorCode::Report(FSCWErrorCode::BadInputFile, TEXT("Missing input file!"));
+		}
+	}
+
+	WorkerDiagnostics.ErrorCode = FSCWErrorCode::IsSet() ? FSCWErrorCode::Get() : FSCWErrorCode::Success;
+
+	static FString LocalHostName = []
+	{
+		FString Output;
+		ISocketSubsystem::Get()->GetHostName(Output);
+		return Output;
+	}();
+
+	TArray<uint8> MemBlock;
+	FMemoryWriter MemWriter(MemBlock);
+	FShaderCompileWorkerUtil::WriteToOutputArchive(MemWriter, LocalHostName, WorkerDiagnostics, NumProcessedJobs, SingleJobs, PipelineJobs, PipelineJobNames);
+
+	static const FValueId OutputId = FValueId::FromName(ANSITEXTVIEW("Output"));
+	Context.AddValue(OutputId, MakeSharedBufferFromArray(MoveTemp(MemBlock)));
+
+	// Reset error code as it can receive a new value in subsequent builds.
+	FSCWErrorCode::Reset();
+}
+
+static UE::DerivedData::TBuildFunctionFactory<FCompileShaderJobsBuildFunction> CompileShaderJobsFunction;
+
+DERIVEDDATABUILDWORKER_API extern int32 BuildWorkerMain(int32 ArgC, TCHAR* ArgV[], const TCHAR* ExtraArgs, const TCHAR* LoadModules);
+
+INT32_MAIN_INT32_ARGC_TCHAR_ARGV()
+{
+	return BuildWorkerMain(ArgC, ArgV, TEXT("-Unattended -ReduceThreadUsage -CpuProfilerTrace -NoCrashReports"), TEXT("ShaderPreprocessor"));
+}
diff --git a/Engine/Source/Programs/ShaderBuildWorker/ShaderBuildWorker.Build.cs b/Engine/Source/Programs/ShaderBuildWorker/ShaderBuildWorker.Build.cs
new file mode 100644
--- /dev/null
+++ b/Engine/Source/Programs/ShaderBuildWorker/ShaderBuildWorker.Build.cs
@@ -0,0 +1,33 @@
+// Copyright Epic Games, Inc. All Rights Reserved.
+
+using UnrealBuildTool;
+using System.IO;
+
+public class ShaderBuildWorker : ModuleRules
+{
+	public ShaderBuildWorker(ReadOnlyTargetRules Target) : base(Target)
+	{
+		bTreatAsEngineModule = true;
+
+		PrivateDependencyModuleNames.AddRange(
+			new string[] {
+				"ApplicationCore",
+				"Core",
+				"DerivedDataBuildWorker",
+				"DerivedDataCache",
+				"Projects",
+				"RenderCore",
+				"RHI",
+				"SandboxFile",
+				"ShaderCompilerCommon",
+				"Sockets",
+				"TargetPlatform",
+				"TraceLog",
+			});
+
+		if (Target.Platform == UnrealTargetPlatform.Win64)
+		{
+			RuntimeDependencies.Add(Path.Combine(Target.WindowsPlatform.DirectXDllDir, "d3dcompiler_47.dll"));
+		}
+	}
+}
diff --git a/Engine/Source/Programs/ShaderBuildWorker/ShaderBuildWorker.Target.cs b/Engine/Source/Programs/ShaderBuildWorker/ShaderBuildWorker.Target.cs
new file mode 100644
--- /dev/null
+++ b/Engine/Source/Programs/ShaderBuildWorker/ShaderBuildWorker.Target.cs
@@ -0,0 +1,29 @@
+// Copyright Epic Games, Inc. All Rights Reserved.
+
+using System.Collections.Generic;
+using System.IO;
+using UnrealBuildTool;
+
+[SupportedPlatforms(UnrealPlatformClass.Editor)]
+public class ShaderBuildWorkerTarget : DerivedDataBuildWorkerTarget
+{
+	public ShaderBuildWorkerTarget(TargetInfo Target) : base(Target)
+	{
+		IncludeOrderVersion = EngineIncludeOrderVersion.Latest;
+		BuildEnvironment = TargetBuildEnvironment.UniqueIfNeeded;
+		LinkType = TargetLinkType.Modular;
+		LaunchModuleName = "ShaderBuildWorker";
+
+		bCompileAgainstApplicationCore = true;
+
+		// Force all shader formats to be built and included.
+		bForceBuildShaderFormats = true;
+
+		if (bShaderCompilerWorkerTrace)
+		{
+			GlobalDefinitions.Add("LLM_ENABLED_IN_CONFIG=1");
+			GlobalDefinitions.Add("UE_MEMORY_TAGS_TRACE_ENABLED=1");
+			bEnableTrace = true;
+		}
+	}
+}
